[mutable] Implemented a collection of nodes to make possible the generation of mutable SKMs from within a Dataflow graph.
- Created a new experimental plugin for all dataflow-related mutable stuff.
- Added node that takes a CO and a series of params and generates a COI ( FCOInstanceGeneratorNode). The output is an array of structures containing a ComponentName and a skeletal mesh. A secondary output has been added containing only the skeletal meshes in case the component they come from is not important.
- Added nodes to be able to define mutable parameters (Material, Mesh and Texture) so they can be set prior to the call for the mutable resources generation. Setting the parameter values externally seems to only work with UObjects. A workaround could be added (there already exists a task for that).
- Added a third kind of node (FCOInstanceGetComponentMesh) that allows the user to select what skeletal mesh wants to use from the ones provided by the output of the " FCustomizableObjectInstanceGeneratorNode" node.

Note : This is the first implementation of the nodes. The update one is currently flawed as the evaluation can not be made asyncronously. That is ought to change in the future once that DF feature is implemented.

#rnx

diff --git a/Engine/Plugins/Experimental/DataflowMutable/DataflowMutable.uplugin b/Engine/Plugins/Experimental/DataflowMutable/DataflowMutable.uplugin
new file mode 100644
--- /dev/null
+++ b/Engine/Plugins/Experimental/DataflowMutable/DataflowMutable.uplugin
@@ -0,0 +1,38 @@
+{
+	"FileVersion": 3,
+	"Version": 1,
+	"VersionName": "0.1",
+	"FriendlyName": "Dataflow Extensions For Mutable",
+	"Description": "Adds Mutable functionality to work with Dataflow objects from the Dataflow plugin",
+	"Category": "Animation",
+	"CreatedBy": "Epic Games, Inc.",
+	"CreatedByURL": "https://epicgames.com",
+	"DocsURL": "",
+	"MarketplaceURL": "",
+	"SupportURL": "",
+	"CanContainContent": false,
+	"IsBetaVersion": false,
+	"IsExperimentalVersion": true,
+	"Installed": false,
+	"EnabledByDefault": true,
+	"Modules": 
+	[
+		{
+			"Name": "DataflowMutableEditor",
+			"Type": "UncookedOnly",
+			"LoadingPhase": "Default"
+		}
+	],
+	"Plugins": 
+	[
+		{
+			"Name": "Dataflow",
+			"Enabled": true
+		},
+		{
+			"Name": "Mutable",
+			"Enabled": true
+		}
+	]
+
+}
\ No newline at end of file
diff --git a/Engine/Plugins/Experimental/DataflowMutable/Source/DataflowMutableEditor/DataflowMutableEditor.Build.cs b/Engine/Plugins/Experimental/DataflowMutable/Source/DataflowMutableEditor/DataflowMutableEditor.Build.cs
new file mode 100644
--- /dev/null
+++ b/Engine/Plugins/Experimental/DataflowMutable/Source/DataflowMutableEditor/DataflowMutableEditor.Build.cs
@@ -0,0 +1,31 @@
+// Copyright Epic Games, Inc. All Rights Reserved.
+
+namespace UnrealBuildTool.Rules
+{
+	public class DataflowMutableEditor : ModuleRules
+	{
+        public DataflowMutableEditor(ReadOnlyTargetRules Target) : base(Target)
+		{
+			PrivateIncludePaths.Add(ModuleDirectory + "/Private");
+			PublicIncludePaths.Add(ModuleDirectory + "/Public");
+			
+			PublicDependencyModuleNames.AddRange(
+				new string[]
+				{
+					"Core",
+					"CoreUObject",
+					"Engine",
+					"CustomizableObject",
+				});
+
+			PrivateDependencyModuleNames.AddRange(
+				new string[]
+				{
+					"DataflowCore",
+					"DataflowEditor",
+					"DataflowEngine",
+					"DataflowNodes",
+				});
+		}
+	}
+}
diff --git a/Engine/Plugins/Experimental/DataflowMutable/Source/DataflowMutableEditor/Private/DataflowMutableEditorModule.cpp b/Engine/Plugins/Experimental/DataflowMutable/Source/DataflowMutableEditor/Private/DataflowMutableEditorModule.cpp
new file mode 100644
--- /dev/null
+++ b/Engine/Plugins/Experimental/DataflowMutable/Source/DataflowMutableEditor/Private/DataflowMutableEditorModule.cpp
@@ -0,0 +1,23 @@
+// Copyright Epic Games, Inc. All Rights Reserved.
+
+#include "DataflowMutableEditorModule.h"
+
+#include "Dataflow/DataflowNode.h"
+#include "Dataflow/DataflowCore.h"
+#include "Nodes/COInstanceGeneratorNode.h"
+
+IMPLEMENT_MODULE(FDataflowMutableEditorModule, DataflowMutableEditor);
+
+void FDataflowMutableEditorModule::StartupModule()
+{
+	// Mian nodeS
+	DATAFLOW_NODE_REGISTER_CREATION_FACTORY(FCOInstanceGeneratorNode);
+
+	// Support nodes
+	DATAFLOW_NODE_REGISTER_CREATION_FACTORY(FCOInstanceGetComponentMesh);
+
+	// Parameter nodes
+	DATAFLOW_NODE_REGISTER_CREATION_FACTORY(FCOInstanceMeshParameterNode);
+	DATAFLOW_NODE_REGISTER_CREATION_FACTORY(FCOInstanceTextureParameterNode);
+	DATAFLOW_NODE_REGISTER_CREATION_FACTORY(FCOInstanceMaterialParameterNode);
+}
diff --git a/Engine/Plugins/Experimental/DataflowMutable/Source/DataflowMutableEditor/Private/Nodes/COInstanceGeneratorNode.cpp b/Engine/Plugins/Experimental/DataflowMutable/Source/DataflowMutableEditor/Private/Nodes/COInstanceGeneratorNode.cpp
new file mode 100644
--- /dev/null
+++ b/Engine/Plugins/Experimental/DataflowMutable/Source/DataflowMutableEditor/Private/Nodes/COInstanceGeneratorNode.cpp
@@ -0,0 +1,415 @@
+// Copyright Epic Games, Inc. All Rights Reserved.
+
+#include "../Public/Nodes/COInstanceGeneratorNode.h"
+
+#include "MuCO/CustomizableObject.h"
+#include "MuCO/CustomizableObjectInstance.h"
+#include "Engine/SkeletalMesh.h"
+#include "MutableDataflowParameters.h"
+
+
+FCOInstanceMeshParameterNode::FCOInstanceMeshParameterNode(const UE::Dataflow::FNodeParameters& InParam, FGuid InGuid)
+	: FDataflowNode(InParam, InGuid)
+{
+	RegisterInputConnection(&ParameterName);
+	RegisterInputConnection(&SkeletalMesh);
+
+	RegisterOutputConnection(&MeshParameter);
+}
+
+
+void FCOInstanceMeshParameterNode::Evaluate(UE::Dataflow::FContext& Context, const FDataflowOutput* Out) const
+{
+	FMutableSkeletalMeshParameter Output;
+	Output.ParameterName = GetValue(Context, &ParameterName);
+	Output.Mesh = GetValue(Context, &SkeletalMesh);
+
+	if (!Output.ParameterName.IsEmpty())
+	{
+		SetValue(Context, Output, &MeshParameter);
+	}
+	else
+	{
+		Context.Warning(TEXT("The generated Mesh Parameter is not valid. Make sure all fields are set."), this);
+	}
+}
+
+
+FCOInstanceTextureParameterNode::FCOInstanceTextureParameterNode(const UE::Dataflow::FNodeParameters& InParam, FGuid InGuid)
+	: FDataflowNode(InParam, InGuid)
+{
+	RegisterInputConnection(&ParameterName);
+	RegisterInputConnection(&Texture);
+
+	RegisterOutputConnection(&TextureParameter);
+}
+
+
+void FCOInstanceTextureParameterNode::Evaluate(UE::Dataflow::FContext& Context, const FDataflowOutput* Out) const
+{
+	FMutableTextureParameter Output;
+	Output.ParameterName = GetValue(Context, &ParameterName);
+	Output.Texture = GetValue(Context, &Texture);
+
+	if (!Output.ParameterName.IsEmpty())
+	{
+		SetValue(Context, Output, &TextureParameter);
+	}
+	else
+	{
+		Context.Warning(TEXT("The generated Texture Parameter is not valid. Make sure all fields are set."), this);
+	}
+}
+
+
+FCOInstanceMaterialParameterNode::FCOInstanceMaterialParameterNode(const UE::Dataflow::FNodeParameters& InParam, FGuid InGuid)
+: FDataflowNode(InParam, InGuid)
+{
+	RegisterInputConnection(&ParameterName);
+	RegisterInputConnection(&Material);
+
+	RegisterOutputConnection(&MaterialParameter);
+}
+
+
+void FCOInstanceMaterialParameterNode::Evaluate(UE::Dataflow::FContext& Context, const FDataflowOutput* Out) const
+{
+	FMutableMaterialParameter Output;
+	Output.ParameterName = GetValue(Context, &ParameterName);
+	Output.Material = GetValue(Context, &Material);
+
+	if (!Output.ParameterName.IsEmpty())
+	{
+		SetValue(Context, Output, &MaterialParameter);
+	}
+	else
+	{
+		Context.Warning(TEXT("The generated Material Parameter is not valid. Make sure all fields are set."), this);
+	}
+}
+
+
+FCOInstanceGetComponentMesh::FCOInstanceGetComponentMesh(const UE::Dataflow::FNodeParameters& InParam, FGuid InGuid)
+	: Super(InParam, InGuid)
+{
+	RegisterInputConnection(&GeneratedResources);
+	RegisterInputConnection(&ComponentName);
+
+	RegisterOutputConnection(&ComponentSkeletalMesh);
+}
+
+
+void FCOInstanceGetComponentMesh::Evaluate(UE::Dataflow::FContext& Context, const FDataflowOutput* Out) const
+{
+	const TArray<FMutableGeneratedResource> InputGeneratedResources = GetValue(Context, &GeneratedResources);
+	const FString TargetComponentName = GetValue(Context, &ComponentName);
+	
+	if (!TargetComponentName.IsEmpty())
+	{
+		for (const FMutableGeneratedResource& Resource : InputGeneratedResources)
+		{
+			if (Resource.ComponentName.Equals(TargetComponentName))
+			{
+				SetValue(Context, Resource.SkeletalMesh, &ComponentSkeletalMesh);
+				return;
+			}
+		}
+	}
+	
+	// No component was found with the provided name
+	SetValue(Context, TObjectPtr<USkeletalMesh>{}, &ComponentSkeletalMesh);
+}
+
+
+void FCOInstanceGetComponentMesh::OnPropertyChanged(UE::Dataflow::FContext& Context,
+	const FPropertyChangedEvent& InPropertyChangedEvent)
+{
+	FDataflowNode::OnPropertyChanged(Context, InPropertyChangedEvent);
+
+	if (const FProperty* ModifiedProperty = InPropertyChangedEvent.Property)
+	{
+		if (ModifiedProperty->GetFName() == GET_MEMBER_NAME_CHECKED(FCOInstanceGetComponentMesh, ComponentName))
+		{
+			Invalidate();
+		}
+	}
+}
+
+
+FCOInstanceGeneratorNode::FCOInstanceGeneratorNode(const UE::Dataflow::FNodeParameters& InParam,FGuid InGuid)
+: FDataflowNode(InParam, InGuid),
+	GenerateInstanceResources( FDataflowFunctionProperty::FDelegate::CreateRaw(this, &FCOInstanceGeneratorNode::OnGenerateInstanceResourcesRequested))
+{
+	// Inputs
+	RegisterInputConnection(&CustomizableObject);
+	RegisterInputConnection(&SkeletalMeshParameters);
+	RegisterInputConnection(&TextureParameters);
+	RegisterInputConnection(&MaterialParameters);
+	
+	// Outputs
+	RegisterOutputConnection(&GeneratedResources);
+	RegisterOutputConnection(&GeneratedSkeletalMeshes);
+}
+
+
+void FCOInstanceGeneratorNode::OnPropertyChanged(UE::Dataflow::FContext& Context, const FPropertyChangedEvent& InPropertyChangedEvent)
+{
+	FDataflowNode::OnPropertyChanged(Context, InPropertyChangedEvent);
+
+	if (const FProperty* ModifiedProperty = InPropertyChangedEvent.Property)
+	{
+		if (ModifiedProperty->GetFName() == GET_MEMBER_NAME_CHECKED(FCOInstanceGeneratorNode, CustomizableObject))
+		{
+			ClearCachedParameters();
+			
+			CustomizableObjectInstance = nullptr;
+			
+			GeneratedResources.Reset();
+			GeneratedSkeletalMeshes.Reset();
+			Invalidate();
+		}
+	}
+}
+
+
+void FCOInstanceGeneratorNode::Evaluate(UE::Dataflow::FContext& Context, const FDataflowOutput* Out) const
+{
+	if (!GetValue(Context, &CustomizableObject))
+	{
+		Context.Error(TEXT("Unable to evaluate : No Customizable Object has been set"), this);
+		return;
+	}
+	
+	if (!CustomizableObjectInstance)
+	{
+		Context.Error(TEXT("Unable to evaluate : Press \"Generate Resources\" before evaluating the node"), this);
+		return;
+	}
+
+	if (bIsGeneratingResources)
+	{
+		Context.Warning(TEXT("Unable to evaluate : Resource generation is in process. Please hold"), this);
+		return;
+	}
+	
+	if (Out->IsA<TArray<FMutableGeneratedResource>>(&GeneratedResources))
+	{
+		SetValue(Context, GeneratedResources, &GeneratedResources);
+	}
+	else if (Out->IsA<TArray<TObjectPtr<USkeletalMesh>>>(&GeneratedSkeletalMeshes))
+	{
+		// Extract the generated meshes from the generated resources collection as there we have them already set
+		TArray<TObjectPtr<USkeletalMesh>> SkeletalMeshes;
+		SkeletalMeshes.Reserve(GeneratedResources.Num());
+		for (const FMutableGeneratedResource& Resource : GeneratedResources)
+		{
+			SkeletalMeshes.Add(Resource.SkeletalMesh);
+		}
+		SetValue(Context, SkeletalMeshes, &GeneratedSkeletalMeshes);
+	}
+}
+
+
+void FCOInstanceGeneratorNode::ClearCachedParameters()
+{
+	CachedSkeletalMeshParameters.Reset();
+	CachedMaterialParameters.Reset();
+	CachedTextureParameters.Reset();
+}
+
+
+void FCOInstanceGeneratorNode::OnGenerateInstanceResourcesRequested(UE::Dataflow::FContext& Context)
+{
+	if (bIsGeneratingResources)
+	{
+		Context.Warning(TEXT("Unable to generate resources. Resource generation is already in process"), this);
+		return;
+	}
+	
+	if (GetValue(Context, &CustomizableObject))
+	{
+		// Prevent the scheduling or more compilations while the generation of resources is in process
+		bIsGeneratingResources = true;
+		RequestCompilation(Context);
+	}
+	else
+	{
+		Context.Error(TEXT("Unable to generate resources. The Customizable object has not been set"), this);
+	}
+}
+
+
+void FCOInstanceGeneratorNode::RequestCompilation(UE::Dataflow::FContext& Context)
+{
+	// Capture the inputs of the node (target parameters) so we can apply them after the compilation without us having to rely on the context (as it may be gone)
+	CacheNodeInputs(Context);
+	check(CustomizableObject);
+	
+	FCompileNativeDelegate CompilationDelegate;
+	TSharedRef<FDataflowNode> This = AsShared();
+	CompilationDelegate.BindLambda([This](const FCompileCallbackParams& CallbackParams)->void
+	{
+		FCOInstanceGeneratorNode* HostNode = static_cast<FCOInstanceGeneratorNode*>(&This.Get());
+		check(HostNode);
+			
+		// First generate the instance
+		HostNode->CustomizableObjectInstance = HostNode->CustomizableObject->CreateInstance();
+		check(HostNode->CustomizableObjectInstance)
+			
+		// Apply the parameters set in the node into the instance
+		HostNode->ApplyInstanceParameters(*HostNode->CustomizableObjectInstance);
+
+		// once the compilation has been performed (if required) ask for the update of the COI
+		HostNode->RequestUpdate(*HostNode->CustomizableObjectInstance);
+	});
+
+	// Request the compilation of the CO
+	FCompileParams CompilationParams;
+	CompilationParams.bAsync = true;
+	CompilationParams.CallbackNative = CompilationDelegate;
+	
+	// todo: Expose compilation errors and warnings to the Dataflow context
+	CustomizableObject->Compile(CompilationParams);
+
+	// Prevent the scheduling or more compilations while the generation of resources is in process
+	bIsGeneratingResources = true;
+}
+
+
+void FCOInstanceGeneratorNode::CacheNodeInputs(UE::Dataflow::FContext& Context)
+{
+	ClearCachedParameters();
+
+	// Customizable Object
+	CustomizableObject = GetValue(Context, &CustomizableObject);
+
+	// Cache the parameters now that we have the context. Do not cache duplicate entries (based on the param name)
+	
+	// Texture Parameters
+	{
+		const TArray<FMutableTextureParameter>& TempTextureParametersArray = GetValue(Context, &TextureParameters);
+		for (const FMutableTextureParameter& TextureParameter : TempTextureParametersArray)
+		{
+			if (CachedTextureParameters.Contains(TextureParameter))
+			{
+				Context.Warning(FString::Printf( TEXT("Multiple Texture Parameters have been found sharing the \"%s\" name , only the first parameter will be used."), *TextureParameter.ParameterName), this);
+				continue;
+			}
+		
+			CachedTextureParameters.Add(TextureParameter);
+		}
+	}
+	
+	// Mesh Parameters
+	{
+		const TArray<FMutableSkeletalMeshParameter>& TempMeshParametersArray = GetValue(Context, &SkeletalMeshParameters);
+		for (const FMutableSkeletalMeshParameter& MeshParameter : TempMeshParametersArray)
+		{
+			if (CachedSkeletalMeshParameters.Contains(MeshParameter))
+			{
+				Context.Warning(FString::Printf( TEXT("Multiple Skeletal Mesh Parameters have been found sharing the \"%s\" name , only the first parameter will be used."), *MeshParameter.ParameterName), this);
+				continue;
+			}
+		
+			CachedSkeletalMeshParameters.Add(MeshParameter);
+		}
+	}
+	
+	// Material parameters
+	{
+		const TArray<FMutableMaterialParameter>& TempMaterialParametersArray = GetValue(Context, &MaterialParameters);
+		for (const FMutableMaterialParameter& MaterialParameter : TempMaterialParametersArray)
+		{
+			if (CachedMaterialParameters.Contains(MaterialParameter))
+			{
+				Context.Warning(FString::Printf( TEXT("Multiple Material Parameters have been found sharing the \"%s\" name , only the first parameter will be used."), *MaterialParameter.ParameterName), this);
+				continue;
+			}
+		
+			CachedMaterialParameters.Add(MaterialParameter);
+		}
+	}
+}
+
+
+void FCOInstanceGeneratorNode::ApplyInstanceParameters(UCustomizableObjectInstance& Instance)
+{
+	// todo. notify when a parameter does not exist :
+	// In order to do that we will require to have async node evaluations and then access to the
+	// context to log the warning
+	
+	// Texture Parameters
+	for (const FMutableTextureParameter& TextureParameter : CachedTextureParameters)
+	{
+		const FString ParameterName = TextureParameter.ParameterName;
+		if (Instance.ContainsTextureParameter(ParameterName))
+		{
+			Instance.SetTextureParameterSelectedOption(ParameterName, TextureParameter.Texture);
+		}
+	}
+	
+	// Mesh Parameters
+	for (const FMutableSkeletalMeshParameter& MeshParameter : CachedSkeletalMeshParameters)
+	{
+		const FString ParameterName = MeshParameter.ParameterName;
+		if (Instance.ContainsSkeletalMeshParameter(ParameterName))
+		{
+			Instance.SetSkeletalMeshParameterSelectedOption(ParameterName, MeshParameter.Mesh);
+		}
+	}
+	
+	// Material parameters
+	for (const FMutableMaterialParameter& MaterialParameter : CachedMaterialParameters)
+	{
+		const FString ParameterName = MaterialParameter.ParameterName;
+		if (Instance.ContainsMaterialParameter(ParameterName))
+		{
+			Instance.SetMaterialParameterSelectedOption(ParameterName, MaterialParameter.Material);				
+		}
+	}
+}
+
+
+void FCOInstanceGeneratorNode::RequestUpdate(UCustomizableObjectInstance& Instance)
+{
+	TSharedRef<FDataflowNode> This = AsShared();
+	
+	// Instance update delegate	
+	FInstanceUpdateNativeDelegate InstanceUpdateNativeDelegate;
+	InstanceUpdateNativeDelegate.AddLambda([This](FUpdateContext OutUpdateContext) -> void
+	{
+		FCOInstanceGeneratorNode* HostNode = static_cast<FCOInstanceGeneratorNode*>(&This.Get());
+		check(HostNode);
+			
+		HostNode->GeneratedResources = HostNode->GetInstanceGeneratedMeshes(*OutUpdateContext.Instance);
+			
+		// Tell the system it is ok to generate resources again
+		HostNode->bIsGeneratingResources = false;
+			
+		This->Invalidate();
+	});
+	 		
+	Instance.UpdateSkeletalMeshAsyncResult(InstanceUpdateNativeDelegate,true, true);
+}
+
+
+TArray<FMutableGeneratedResource> FCOInstanceGeneratorNode::GetInstanceGeneratedMeshes(const UCustomizableObjectInstance& Instance) const
+{
+	TArray<FMutableGeneratedResource> GeneratedMeshesData;
+			
+	const TArray<FName> InstanceComponentNames = Instance.GetComponentNames();
+	for (const FName& ComponentName : InstanceComponentNames)
+	{
+		if (const TObjectPtr<USkeletalMesh> SkeletalMesh = Instance.GetSkeletalMeshComponentSkeletalMesh(ComponentName))
+		{
+			FMutableGeneratedResource ComponentData;
+			ComponentData.ComponentName = ComponentName.ToString();
+			ComponentData.SkeletalMesh = SkeletalMesh;
+			
+			GeneratedMeshesData.Add(ComponentData);
+		}
+	}
+
+	return GeneratedMeshesData;
+}
+
diff --git a/Engine/Plugins/Experimental/DataflowMutable/Source/DataflowMutableEditor/Public/DataflowMutableEditorModule.h b/Engine/Plugins/Experimental/DataflowMutable/Source/DataflowMutableEditor/Public/DataflowMutableEditorModule.h
new file mode 100644
--- /dev/null
+++ b/Engine/Plugins/Experimental/DataflowMutable/Source/DataflowMutableEditor/Public/DataflowMutableEditorModule.h
@@ -0,0 +1,9 @@
+// Copyright Epic Games, Inc. All Rights Reserved.
+#pragma once
+
+#include "Modules/ModuleInterface.h"
+
+class FDataflowMutableEditorModule : public IModuleInterface
+{
+	virtual void StartupModule() override;
+};
diff --git a/Engine/Plugins/Experimental/DataflowMutable/Source/DataflowMutableEditor/Public/MutableDataflowParameters.h b/Engine/Plugins/Experimental/DataflowMutable/Source/DataflowMutableEditor/Public/MutableDataflowParameters.h
new file mode 100644
--- /dev/null
+++ b/Engine/Plugins/Experimental/DataflowMutable/Source/DataflowMutableEditor/Public/MutableDataflowParameters.h
@@ -0,0 +1,57 @@
+// Copyright Epic Games, Inc. All Rights Reserved.
+#pragma once
+
+#include "CoreMinimal.h"
+// #include "UObject/ObjectMacros.h"
+
+#include "MutableDataflowParameters.generated.h"
+
+class USkeletalMesh;
+class UTexture2D;
+class UMaterialInterface;
+
+USTRUCT()
+struct FMutableParameterBase
+{
+	GENERATED_USTRUCT_BODY()
+
+	virtual ~FMutableParameterBase() = default;
+	
+	UPROPERTY(EditAnywhere, Category = "Mutable");
+	FString ParameterName;
+	
+	bool operator== (const FMutableParameterBase& Other) const
+	{
+		return this->ParameterName.ToLower() == Other.ParameterName.ToLower();
+	}
+};
+
+
+USTRUCT(BlueprintType, Meta = (Experimental))
+struct FMutableSkeletalMeshParameter : public FMutableParameterBase
+{
+	GENERATED_USTRUCT_BODY()
+	
+	UPROPERTY(EditAnywhere, Category = "Mutable");
+	TObjectPtr<USkeletalMesh> Mesh;
+};
+
+
+USTRUCT(BlueprintType, Meta = (Experimental))
+struct FMutableTextureParameter : public FMutableParameterBase
+{
+	GENERATED_USTRUCT_BODY()
+	
+	UPROPERTY(EditAnywhere, Category = "Mutable");
+	TObjectPtr<UTexture2D> Texture;
+};
+
+
+USTRUCT(BlueprintType, Meta = (Experimental))
+struct FMutableMaterialParameter : public FMutableParameterBase
+{
+	GENERATED_USTRUCT_BODY()
+	
+	UPROPERTY(EditAnywhere, Category = "Mutable");
+	TObjectPtr<UMaterialInterface> Material;
+};
\ No newline at end of file
diff --git a/Engine/Plugins/Experimental/DataflowMutable/Source/DataflowMutableEditor/Public/Nodes/COInstanceGeneratorNode.h b/Engine/Plugins/Experimental/DataflowMutable/Source/DataflowMutableEditor/Public/Nodes/COInstanceGeneratorNode.h
new file mode 100644
--- /dev/null
+++ b/Engine/Plugins/Experimental/DataflowMutable/Source/DataflowMutableEditor/Public/Nodes/COInstanceGeneratorNode.h
@@ -0,0 +1,225 @@
+// Copyright Epic Games, Inc. All Rights Reserved.
+#pragma once
+
+#include "Dataflow/DataflowNode.h"
+#include "Misc/Guid.h"
+#include "CoreMinimal.h"
+#include "Dataflow/DataflowFunctionProperty.h"
+#include "MuCO/CustomizableObject.h"
+#include "Engine/SkeletalMesh.h"
+#include "GeometryCollection/ManagedArrayCollection.h"
+#include "Dataflow/DataflowInputOutput.h"
+#include "Materials/MaterialInterface.h"
+#include "MutableDataflowParameters.h"
+
+#include "COInstanceGeneratorNode.generated.h"
+
+USTRUCT(Meta = (Experimental))
+struct FMutableGeneratedResource
+{
+	GENERATED_USTRUCT_BODY()
+
+	UPROPERTY(EditAnywhere, Category = "Mutable");
+	FString ComponentName;
+
+	UPROPERTY(EditAnywhere, Category = "Mutable");
+	TObjectPtr<USkeletalMesh> SkeletalMesh;
+};
+
+
+USTRUCT(Meta = (Experimental))
+struct FCOInstanceMeshParameterNode : public FDataflowNode
+{
+	GENERATED_USTRUCT_BODY()
+	DATAFLOW_NODE_DEFINE_INTERNAL(FCOInstanceMeshParameterNode, "Mutable Mesh Parameter", "Mutable", "")
+	
+	UPROPERTY(EditAnywhere, Category = "Mutable", meta = (DataflowInput));
+	FString ParameterName;
+
+	UPROPERTY(EditAnywhere, Category = "Mutable", meta = (DataflowInput));
+	TObjectPtr<USkeletalMesh> SkeletalMesh;
+	
+	UPROPERTY(meta = (DataflowOutput));
+	FMutableSkeletalMeshParameter MeshParameter;
+
+public:
+	
+	FCOInstanceMeshParameterNode(const UE::Dataflow::FNodeParameters& InParam, FGuid InGuid = FGuid::NewGuid());
+	
+	virtual void Evaluate(UE::Dataflow::FContext& Context, const FDataflowOutput* Out) const override;
+};
+
+
+USTRUCT(Meta = (Experimental))
+struct FCOInstanceTextureParameterNode : public FDataflowNode
+{
+	GENERATED_USTRUCT_BODY()
+	DATAFLOW_NODE_DEFINE_INTERNAL(FCOInstanceTextureParameterNode, "Mutable Texture Parameter", "Mutable", "")
+
+	UPROPERTY(EditAnywhere, Category = "Mutable", meta = (DataflowInput));
+	FString ParameterName;
+
+	UPROPERTY(EditAnywhere, Category = "Mutable", meta = (DataflowInput));
+	TObjectPtr<UTexture2D> Texture;
+	
+	UPROPERTY(meta = (DataflowOutput));
+	FMutableTextureParameter TextureParameter;
+	
+public:
+	
+	FCOInstanceTextureParameterNode(const UE::Dataflow::FNodeParameters& InParam, FGuid InGuid = FGuid::NewGuid());
+
+	virtual void Evaluate(UE::Dataflow::FContext& Context, const FDataflowOutput* Out) const override;
+};
+
+
+USTRUCT(Meta = (Experimental))
+struct FCOInstanceMaterialParameterNode : public FDataflowNode
+{
+	GENERATED_USTRUCT_BODY()
+	DATAFLOW_NODE_DEFINE_INTERNAL(FCOInstanceMaterialParameterNode, "Mutable Material Parameter", "Mutable", "")
+
+	UPROPERTY(EditAnywhere, Category = "Mutable", meta = (DataflowInput));
+	FString ParameterName;
+
+	UPROPERTY(EditAnywhere, Category = "Mutable", meta = (DataflowInput));
+	TObjectPtr<UMaterialInterface> Material;
+
+	UPROPERTY(meta = (DataflowOutput));
+	FMutableMaterialParameter MaterialParameter;
+
+public:
+	
+	FCOInstanceMaterialParameterNode(const UE::Dataflow::FNodeParameters& InParam, FGuid InGuid = FGuid::NewGuid());
+	
+	virtual void Evaluate(UE::Dataflow::FContext& Context, const FDataflowOutput* Out) const override;
+};
+
+
+USTRUCT(Meta = (Experimental))
+struct FCOInstanceGetComponentMesh : public FDataflowNode
+{
+	GENERATED_USTRUCT_BODY()
+	DATAFLOW_NODE_DEFINE_INTERNAL(FCOInstanceGetComponentMesh, "Get Component Mesh", "Mutable", "")
+
+	UPROPERTY(EditAnywhere, Category = "Mutable", meta = (DataflowInput))
+	FString ComponentName;
+
+	UPROPERTY(meta = (DataflowInput))
+	TArray<FMutableGeneratedResource> GeneratedResources;
+
+	UPROPERTY(meta = (DataflowOutput))
+	TObjectPtr<USkeletalMesh> ComponentSkeletalMesh;
+
+public:
+	
+	FCOInstanceGetComponentMesh(const UE::Dataflow::FNodeParameters& InParam, FGuid InGuid = FGuid::NewGuid());
+	
+	virtual void Evaluate(UE::Dataflow::FContext& Context, const FDataflowOutput* Out) const override;
+	
+	virtual void OnPropertyChanged(UE::Dataflow::FContext& Context, const FPropertyChangedEvent& InPropertyChangedEvent) override;
+};
+
+
+/**
+ * Node designed to work as an entry point for mutable in the DataFlow graph.
+ * It allows the setup of a Customizable Object Instance and the extraction of the skeletal meshes generated by its update.
+ */
+USTRUCT(Meta = (Experimental))
+struct FCOInstanceGeneratorNode : public FDataflowNode
+{
+	GENERATED_USTRUCT_BODY()
+	DATAFLOW_NODE_DEFINE_INTERNAL(FCOInstanceGeneratorNode, "Update COI", "Mutable", "")
+
+public:
+	
+	FCOInstanceGeneratorNode(const UE::Dataflow::FNodeParameters& InParam, FGuid InGuid = FGuid::NewGuid());
+
+	virtual void OnPropertyChanged(UE::Dataflow::FContext& Context, const FPropertyChangedEvent& InPropertyChangedEvent) override;
+
+	virtual void Evaluate(UE::Dataflow::FContext& Context, const FDataflowOutput* Out) const override;
+
+private:
+	
+	// Inputs
+
+	UPROPERTY(EditAnywhere, Category = "Mutable", meta = (DataflowInput));
+	TObjectPtr<UCustomizableObject> CustomizableObject;
+
+	
+	UPROPERTY(EditAnywhere, Category = "Mutable Parameters", meta = (DataflowInput))
+	TArray<FMutableSkeletalMeshParameter> SkeletalMeshParameters;
+	TArray<FMutableSkeletalMeshParameter> CachedSkeletalMeshParameters;
+	
+	UPROPERTY(EditAnywhere, Category = "Mutable Parameters", meta = (DataflowInput))
+	TArray<FMutableTextureParameter> TextureParameters;
+	TArray<FMutableTextureParameter> CachedTextureParameters;
+
+	UPROPERTY(EditAnywhere, Category = "Mutable Parameters", meta = (DataflowInput))
+	TArray<FMutableMaterialParameter> MaterialParameters;
+	TArray<FMutableMaterialParameter> CachedMaterialParameters;
+
+	// Outputs
+
+	// First option : an array of structures that contain the pair of data Component name + Skeletal Mesh
+	UPROPERTY(meta = (DataflowOutput))
+	TArray<FMutableGeneratedResource> GeneratedResources;
+	
+	UPROPERTY(meta = (DataflowOutput))
+	TArray<TObjectPtr<USkeletalMesh>> GeneratedSkeletalMeshes;
+
+
+	// Instance generated during the runtime and used to generate the output resources
+	UPROPERTY(EditAnywhere, Category = "Mutable", meta = (DisplayName = "Generated COI"));
+	TObjectPtr<UCustomizableObjectInstance> CustomizableObjectInstance;
+	
+	/** Click on this button to reimport the specified USD file and regenerate the intermediary assets. */
+	UPROPERTY(EditAnywhere, Category = "Mutable", Meta = (DisplayName = "Generate Resources", ButtonImage = "Icons.Refresh"))
+	FDataflowFunctionProperty GenerateInstanceResources;
+	
+	
+	/** Variable that works as a semaphore to prevent the scheduling of multiple Instance Compilation + Updates*/
+	bool bIsGeneratingResources = false;
+	
+	/**
+	 * Callback method invoked when the "Generate outputs" button gets pressed. It will trigger the compilation and later update of the generated 
+	 Customizable Object Instance
+	 */
+	void OnGenerateInstanceResourcesRequested(UE::Dataflow::FContext& Context);
+	
+	/**
+	 * Evaluate the input nodes and add their data to our local caches so we can later access it without the dataflow context
+	 */
+	void CacheNodeInputs(UE::Dataflow::FContext& Context);
+	
+	/**
+	 * Apply the values in the cached parameters into the provided Instance
+	 * @param Instance The instance that will get its parameters updated
+	 */
+	void ApplyInstanceParameters(UCustomizableObjectInstance& Instance);
+
+	/**
+	 * Empty the caches holding the data provided by the evaluation of the input pins
+	 */
+	void ClearCachedParameters();
+	
+	/**
+	 * Request the compilation of the CO so we can later generate a Customizable Object Instance.
+	 */
+	void RequestCompilation(UE::Dataflow::FContext& Context);
+	
+	/** 
+	 * Request the Update of the given COI.
+	 */
+	void RequestUpdate(UCustomizableObjectInstance& Instance);
+
+	/**
+	 * Extract from the given Customizable Object Instance the generated skeletal meshes for each one of the components the instance has.
+	 * @param Instance The COI we want to get the data from
+	 * @return An array of pairs where the key represents the name of the component and the value the skeletal mesh on that component.
+	 */
+	TArray<FMutableGeneratedResource> GetInstanceGeneratedMeshes(const UCustomizableObjectInstance& Instance) const;
+};
+
+
+
