[Mutable] Dataflow made the parameter nodes functional:
- Added a custom node to be able to group parameter nodes in arrays so they can be provided as an input for the COInstanceGeneratorNode.
- Added autoconversions for the cases to handle cases where a single parameter is being connected to an array of parameters input.
- Added autoconversion to convert an array of mutable FMutableGeneratedResource into a single SkeletalMesh pointer using the FCOInstanceGetComponentMesh node as intermediary.
- Moved parameter nodes into their own files.
- Renamed the UI side of the nodes so they are consistent with the other Dataflow nodes

#rnx

diff --git a/Engine/Plugins/Experimental/DataflowMutable/Source/DataflowMutableEditor/Private/DataflowMutableEditorModule.cpp b/Engine/Plugins/Experimental/DataflowMutable/Source/DataflowMutableEditor/Private/DataflowMutableEditorModule.cpp
--- a/Engine/Plugins/Experimental/DataflowMutable/Source/DataflowMutableEditor/Private/DataflowMutableEditorModule.cpp
+++ b/Engine/Plugins/Experimental/DataflowMutable/Source/DataflowMutableEditor/Private/DataflowMutableEditorModule.cpp
@@ -2,22 +2,40 @@
 
 #include "DataflowMutableEditorModule.h"
 
+#include "Dataflow/DataflowAnyTypeRegistry.h"
 #include "Dataflow/DataflowNode.h"
 #include "Dataflow/DataflowCore.h"
 #include "Nodes/COInstanceGeneratorNode.h"
+#include "Nodes/MutableMaterialParameterNode.h"
+#include "Nodes/MutableSkeletalMeshParameterNode.h"
+#include "Nodes/MutableTextureParameterNode.h"
+#include "Nodes/MakeMutableMaterialParametersArrayNode.h"
+#include "Nodes/MakeMutableSkeletalMeshParametersArrayNode.h"
+#include "Nodes/MakeMutableTextureParametersArrayNode.h"
 
 IMPLEMENT_MODULE(FDataflowMutableEditorModule, DataflowMutableEditor);
 
 void FDataflowMutableEditorModule::StartupModule()
 {
-	// Mian nodeS
+	// Main nodes
 	DATAFLOW_NODE_REGISTER_CREATION_FACTORY(FCOInstanceGeneratorNode);
-
-	// Support nodes
 	DATAFLOW_NODE_REGISTER_CREATION_FACTORY(FCOInstanceGetComponentMesh);
 
+	// Auto connection between an array of FMutableGeneratedResource to a single TObjectPtr<USkeletalMesh>
+	UE_DATAFLOW_REGISTER_AUTOCONVERT(TArray<FMutableGeneratedResource>, TObjectPtr<USkeletalMesh>, FCOInstanceGetComponentMesh);
+	
 	// Parameter nodes
-	DATAFLOW_NODE_REGISTER_CREATION_FACTORY(FCOInstanceMeshParameterNode);
-	DATAFLOW_NODE_REGISTER_CREATION_FACTORY(FCOInstanceTextureParameterNode);
-	DATAFLOW_NODE_REGISTER_CREATION_FACTORY(FCOInstanceMaterialParameterNode);
+	DATAFLOW_NODE_REGISTER_CREATION_FACTORY(FMutableSkeletalMeshParameterNode);
+	DATAFLOW_NODE_REGISTER_CREATION_FACTORY(FMutableTextureParameterNode);
+	DATAFLOW_NODE_REGISTER_CREATION_FACTORY(FMutableMaterialParameterNode);
+
+	// Parameter array generation nodes
+	DATAFLOW_NODE_REGISTER_CREATION_FACTORY(FMakeMutableSkeletalMeshParametersArrayNode);
+	DATAFLOW_NODE_REGISTER_CREATION_FACTORY(FMakeMutableMaterialParametersArrayNode);
+	DATAFLOW_NODE_REGISTER_CREATION_FACTORY(FMakeMutableTextureParametersArrayNode);
+
+	// Parameter node to array automatic conversion nodes
+	UE_DATAFLOW_REGISTER_AUTOCONVERT(FMutableSkeletalMeshParameter, TArray<FMutableSkeletalMeshParameter>, FMakeMutableSkeletalMeshParametersArrayNode);
+	UE_DATAFLOW_REGISTER_AUTOCONVERT(FMutableMaterialParameter, TArray<FMutableMaterialParameter>, FMakeMutableMaterialParametersArrayNode);
+	UE_DATAFLOW_REGISTER_AUTOCONVERT(FMutableTextureParameter, TArray<FMutableTextureParameter>, FMakeMutableTextureParametersArrayNode);
 }
diff --git a/Engine/Plugins/Experimental/DataflowMutable/Source/DataflowMutableEditor/Private/Nodes/COInstanceGeneratorNode.cpp b/Engine/Plugins/Experimental/DataflowMutable/Source/DataflowMutableEditor/Private/Nodes/COInstanceGeneratorNode.cpp
--- a/Engine/Plugins/Experimental/DataflowMutable/Source/DataflowMutableEditor/Private/Nodes/COInstanceGeneratorNode.cpp
+++ b/Engine/Plugins/Experimental/DataflowMutable/Source/DataflowMutableEditor/Private/Nodes/COInstanceGeneratorNode.cpp
@@ -8,87 +8,6 @@
 #include "MutableDataflowParameters.h"
 
 
-FCOInstanceMeshParameterNode::FCOInstanceMeshParameterNode(const UE::Dataflow::FNodeParameters& InParam, FGuid InGuid)
-	: FDataflowNode(InParam, InGuid)
-{
-	RegisterInputConnection(&ParameterName);
-	RegisterInputConnection(&SkeletalMesh);
-
-	RegisterOutputConnection(&MeshParameter);
-}
-
-
-void FCOInstanceMeshParameterNode::Evaluate(UE::Dataflow::FContext& Context, const FDataflowOutput* Out) const
-{
-	FMutableSkeletalMeshParameter Output;
-	Output.ParameterName = GetValue(Context, &ParameterName);
-	Output.Mesh = GetValue(Context, &SkeletalMesh);
-
-	if (!Output.ParameterName.IsEmpty())
-	{
-		SetValue(Context, Output, &MeshParameter);
-	}
-	else
-	{
-		Context.Warning(TEXT("The generated Mesh Parameter is not valid. Make sure all fields are set."), this);
-	}
-}
-
-
-FCOInstanceTextureParameterNode::FCOInstanceTextureParameterNode(const UE::Dataflow::FNodeParameters& InParam, FGuid InGuid)
-	: FDataflowNode(InParam, InGuid)
-{
-	RegisterInputConnection(&ParameterName);
-	RegisterInputConnection(&Texture);
-
-	RegisterOutputConnection(&TextureParameter);
-}
-
-
-void FCOInstanceTextureParameterNode::Evaluate(UE::Dataflow::FContext& Context, const FDataflowOutput* Out) const
-{
-	FMutableTextureParameter Output;
-	Output.ParameterName = GetValue(Context, &ParameterName);
-	Output.Texture = GetValue(Context, &Texture);
-
-	if (!Output.ParameterName.IsEmpty())
-	{
-		SetValue(Context, Output, &TextureParameter);
-	}
-	else
-	{
-		Context.Warning(TEXT("The generated Texture Parameter is not valid. Make sure all fields are set."), this);
-	}
-}
-
-
-FCOInstanceMaterialParameterNode::FCOInstanceMaterialParameterNode(const UE::Dataflow::FNodeParameters& InParam, FGuid InGuid)
-: FDataflowNode(InParam, InGuid)
-{
-	RegisterInputConnection(&ParameterName);
-	RegisterInputConnection(&Material);
-
-	RegisterOutputConnection(&MaterialParameter);
-}
-
-
-void FCOInstanceMaterialParameterNode::Evaluate(UE::Dataflow::FContext& Context, const FDataflowOutput* Out) const
-{
-	FMutableMaterialParameter Output;
-	Output.ParameterName = GetValue(Context, &ParameterName);
-	Output.Material = GetValue(Context, &Material);
-
-	if (!Output.ParameterName.IsEmpty())
-	{
-		SetValue(Context, Output, &MaterialParameter);
-	}
-	else
-	{
-		Context.Warning(TEXT("The generated Material Parameter is not valid. Make sure all fields are set."), this);
-	}
-}
-
-
 FCOInstanceGetComponentMesh::FCOInstanceGetComponentMesh(const UE::Dataflow::FNodeParameters& InParam, FGuid InGuid)
 	: Super(InParam, InGuid)
 {
@@ -103,38 +22,24 @@ void FCOInstanceGetComponentMesh::Evaluate(UE::Dataflow::FContext& Context, cons
 {
 	const TArray<FMutableGeneratedResource> InputGeneratedResources = GetValue(Context, &GeneratedResources);
 	const FString TargetComponentName = GetValue(Context, &ComponentName);
-	
+
+	TObjectPtr<USkeletalMesh> OutputSkeletalMesh = nullptr;
 	if (!TargetComponentName.IsEmpty())
 	{
 		for (const FMutableGeneratedResource& Resource : InputGeneratedResources)
 		{
 			if (Resource.ComponentName.Equals(TargetComponentName))
 			{
-				SetValue(Context, Resource.SkeletalMesh, &ComponentSkeletalMesh);
-				return;
+				OutputSkeletalMesh = Resource.SkeletalMesh;
+				break;
 			}
 		}
 	}
 	
-	// No component was found with the provided name
-	SetValue(Context, TObjectPtr<USkeletalMesh>{}, &ComponentSkeletalMesh);
+	SetValue(Context, OutputSkeletalMesh, &ComponentSkeletalMesh);
 }
 
 
-void FCOInstanceGetComponentMesh::OnPropertyChanged(UE::Dataflow::FContext& Context,
-	const FPropertyChangedEvent& InPropertyChangedEvent)
-{
-	FDataflowNode::OnPropertyChanged(Context, InPropertyChangedEvent);
-
-	if (const FProperty* ModifiedProperty = InPropertyChangedEvent.Property)
-	{
-		if (ModifiedProperty->GetFName() == GET_MEMBER_NAME_CHECKED(FCOInstanceGetComponentMesh, ComponentName))
-		{
-			Invalidate();
-		}
-	}
-}
-
 
 FCOInstanceGeneratorNode::FCOInstanceGeneratorNode(const UE::Dataflow::FNodeParameters& InParam,FGuid InGuid)
 : FDataflowNode(InParam, InGuid),
@@ -292,7 +197,7 @@ void FCOInstanceGeneratorNode::CacheNodeInputs(UE::Dataflow::FContext& Context)
 		{
 			if (CachedTextureParameters.Contains(TextureParameter))
 			{
-				Context.Warning(FString::Printf( TEXT("Multiple Texture Parameters have been found sharing the \"%s\" name , only the first parameter will be used."), *TextureParameter.ParameterName), this);
+				Context.Warning(FString::Printf( TEXT("Multiple Texture Parameters have been found sharing the \"%s\" name , only the first parameter will be used."), *TextureParameter.Name), this);
 				continue;
 			}
 		
@@ -307,7 +212,7 @@ void FCOInstanceGeneratorNode::CacheNodeInputs(UE::Dataflow::FContext& Context)
 		{
 			if (CachedSkeletalMeshParameters.Contains(MeshParameter))
 			{
-				Context.Warning(FString::Printf( TEXT("Multiple Skeletal Mesh Parameters have been found sharing the \"%s\" name , only the first parameter will be used."), *MeshParameter.ParameterName), this);
+				Context.Warning(FString::Printf( TEXT("Multiple Skeletal Mesh Parameters have been found sharing the \"%s\" name , only the first parameter will be used."), *MeshParameter.Name), this);
 				continue;
 			}
 		
@@ -322,7 +227,7 @@ void FCOInstanceGeneratorNode::CacheNodeInputs(UE::Dataflow::FContext& Context)
 		{
 			if (CachedMaterialParameters.Contains(MaterialParameter))
 			{
-				Context.Warning(FString::Printf( TEXT("Multiple Material Parameters have been found sharing the \"%s\" name , only the first parameter will be used."), *MaterialParameter.ParameterName), this);
+				Context.Warning(FString::Printf( TEXT("Multiple Material Parameters have been found sharing the \"%s\" name , only the first parameter will be used."), *MaterialParameter.Name), this);
 				continue;
 			}
 		
@@ -341,7 +246,7 @@ void FCOInstanceGeneratorNode::ApplyInstanceParameters(UCustomizableObjectInstan
 	// Texture Parameters
 	for (const FMutableTextureParameter& TextureParameter : CachedTextureParameters)
 	{
-		const FString ParameterName = TextureParameter.ParameterName;
+		const FString ParameterName = TextureParameter.Name;
 		if (Instance.ContainsTextureParameter(ParameterName))
 		{
 			Instance.SetTextureParameterSelectedOption(ParameterName, TextureParameter.Texture);
@@ -351,7 +256,7 @@ void FCOInstanceGeneratorNode::ApplyInstanceParameters(UCustomizableObjectInstan
 	// Mesh Parameters
 	for (const FMutableSkeletalMeshParameter& MeshParameter : CachedSkeletalMeshParameters)
 	{
-		const FString ParameterName = MeshParameter.ParameterName;
+		const FString ParameterName = MeshParameter.Name;
 		if (Instance.ContainsSkeletalMeshParameter(ParameterName))
 		{
 			Instance.SetSkeletalMeshParameterSelectedOption(ParameterName, MeshParameter.Mesh);
@@ -361,7 +266,7 @@ void FCOInstanceGeneratorNode::ApplyInstanceParameters(UCustomizableObjectInstan
 	// Material parameters
 	for (const FMutableMaterialParameter& MaterialParameter : CachedMaterialParameters)
 	{
-		const FString ParameterName = MaterialParameter.ParameterName;
+		const FString ParameterName = MaterialParameter.Name;
 		if (Instance.ContainsMaterialParameter(ParameterName))
 		{
 			Instance.SetMaterialParameterSelectedOption(ParameterName, MaterialParameter.Material);				
diff --git a/Engine/Plugins/Experimental/DataflowMutable/Source/DataflowMutableEditor/Private/Nodes/MakeMutableMaterialParametersArrayNode.cpp b/Engine/Plugins/Experimental/DataflowMutable/Source/DataflowMutableEditor/Private/Nodes/MakeMutableMaterialParametersArrayNode.cpp
new file mode 100644
--- /dev/null
+++ b/Engine/Plugins/Experimental/DataflowMutable/Source/DataflowMutableEditor/Private/Nodes/MakeMutableMaterialParametersArrayNode.cpp
@@ -0,0 +1,50 @@
+﻿// Copyright Epic Games, Inc. All Rights Reserved.
+
+
+#include "Nodes/MakeMutableMaterialParametersArrayNode.h"
+
+
+FMakeMutableMaterialParametersArrayNode::FMakeMutableMaterialParametersArrayNode(const UE::Dataflow::FNodeParameters& InParam, FGuid InGuid)
+	: Super(InParam, InGuid)
+{
+	RegisterOutputConnection(&GroupedParameters);
+
+	// Create as many inputs as the node requires to have by default
+	AddDefaultInputs(InputMaterialParameters);
+}
+
+
+void FMakeMutableMaterialParametersArrayNode::Evaluate(UE::Dataflow::FContext& Context, const FDataflowOutput* Out) const
+{
+	EvaluateParameterNode(InputMaterialParameters, GroupedParameters, Context, Out);
+}
+
+
+TArray<UE::Dataflow::FPin> FMakeMutableMaterialParametersArrayNode::AddPins()
+{
+	return AddParameterPin(InputMaterialParameters);
+}
+
+
+bool FMakeMutableMaterialParametersArrayNode::CanRemovePin() const
+{
+	return CanRemoveParameterPin(InputMaterialParameters);
+}
+
+
+TArray<UE::Dataflow::FPin> FMakeMutableMaterialParametersArrayNode::GetPinsToRemove() const
+{
+	return GetParameterPinsToRemove(InputMaterialParameters);
+}
+
+
+void FMakeMutableMaterialParametersArrayNode::OnPinRemoved(const UE::Dataflow::FPin& Pin)
+{
+	OnParameterPinRemoved(InputMaterialParameters, Pin);
+}
+
+
+void FMakeMutableMaterialParametersArrayNode::PostSerialize(const FArchive& Ar)
+{
+	PostNodeSerialize(InputMaterialParameters, Ar);
+}
diff --git a/Engine/Plugins/Experimental/DataflowMutable/Source/DataflowMutableEditor/Private/Nodes/MakeMutableSkeletalMeshParametersArrayNode.cpp b/Engine/Plugins/Experimental/DataflowMutable/Source/DataflowMutableEditor/Private/Nodes/MakeMutableSkeletalMeshParametersArrayNode.cpp
new file mode 100644
--- /dev/null
+++ b/Engine/Plugins/Experimental/DataflowMutable/Source/DataflowMutableEditor/Private/Nodes/MakeMutableSkeletalMeshParametersArrayNode.cpp
@@ -0,0 +1,50 @@
+// Copyright Epic Games, Inc. All Rights Reserved.
+
+
+#include "Nodes/MakeMutableSkeletalMeshParametersArrayNode.h"
+
+
+FMakeMutableSkeletalMeshParametersArrayNode::FMakeMutableSkeletalMeshParametersArrayNode(const UE::Dataflow::FNodeParameters& InParam, FGuid InGuid)
+	: Super(InParam, InGuid)
+{
+	RegisterOutputConnection(&GroupedParameters);
+
+	// Create as many inputs as the node requires to have by default
+	AddDefaultInputs(InputSkeletalMeshParameters);
+}
+
+
+void FMakeMutableSkeletalMeshParametersArrayNode::Evaluate(UE::Dataflow::FContext& Context, const FDataflowOutput* Out) const
+{
+	EvaluateParameterNode(InputSkeletalMeshParameters, GroupedParameters, Context, Out);
+}
+
+
+TArray<UE::Dataflow::FPin> FMakeMutableSkeletalMeshParametersArrayNode::AddPins()
+{
+	return AddParameterPin(InputSkeletalMeshParameters);
+}
+
+
+bool FMakeMutableSkeletalMeshParametersArrayNode::CanRemovePin() const
+{
+	return CanRemoveParameterPin(InputSkeletalMeshParameters);
+}
+
+
+TArray<UE::Dataflow::FPin> FMakeMutableSkeletalMeshParametersArrayNode::GetPinsToRemove() const
+{
+	return GetParameterPinsToRemove(InputSkeletalMeshParameters);
+}
+
+
+void FMakeMutableSkeletalMeshParametersArrayNode::OnPinRemoved(const UE::Dataflow::FPin& Pin)
+{
+	OnParameterPinRemoved(InputSkeletalMeshParameters, Pin);
+}
+
+
+void FMakeMutableSkeletalMeshParametersArrayNode::PostSerialize(const FArchive& Ar)
+{
+	PostNodeSerialize(InputSkeletalMeshParameters, Ar);
+}
diff --git a/Engine/Plugins/Experimental/DataflowMutable/Source/DataflowMutableEditor/Private/Nodes/MakeMutableTextureParametersArrayNode.cpp b/Engine/Plugins/Experimental/DataflowMutable/Source/DataflowMutableEditor/Private/Nodes/MakeMutableTextureParametersArrayNode.cpp
new file mode 100644
--- /dev/null
+++ b/Engine/Plugins/Experimental/DataflowMutable/Source/DataflowMutableEditor/Private/Nodes/MakeMutableTextureParametersArrayNode.cpp
@@ -0,0 +1,50 @@
+﻿// Copyright Epic Games, Inc. All Rights Reserved.
+
+
+#include "Nodes/MakeMutableTextureParametersArrayNode.h"
+
+
+FMakeMutableTextureParametersArrayNode::FMakeMutableTextureParametersArrayNode(const UE::Dataflow::FNodeParameters& InParam, FGuid InGuid)
+	: Super(InParam, InGuid)
+{
+	RegisterOutputConnection(&GroupedParameters);
+
+	// Create as many inputs as the node requires to have by default
+	AddDefaultInputs(InputTextureParameters);
+}
+
+
+void FMakeMutableTextureParametersArrayNode::Evaluate(UE::Dataflow::FContext& Context, const FDataflowOutput* Out) const
+{
+	EvaluateParameterNode(InputTextureParameters, GroupedParameters, Context, Out);
+}
+
+
+TArray<UE::Dataflow::FPin> FMakeMutableTextureParametersArrayNode::AddPins()
+{
+	return AddParameterPin(InputTextureParameters);
+}
+
+
+bool FMakeMutableTextureParametersArrayNode::CanRemovePin() const
+{
+	return CanRemoveParameterPin(InputTextureParameters);
+}
+
+
+TArray<UE::Dataflow::FPin> FMakeMutableTextureParametersArrayNode::GetPinsToRemove() const
+{
+	return GetParameterPinsToRemove(InputTextureParameters);
+}
+
+
+void FMakeMutableTextureParametersArrayNode::OnPinRemoved(const UE::Dataflow::FPin& Pin)
+{
+	OnParameterPinRemoved(InputTextureParameters, Pin);
+}
+
+
+void FMakeMutableTextureParametersArrayNode::PostSerialize(const FArchive& Ar)
+{
+	PostNodeSerialize(InputTextureParameters, Ar);
+}
diff --git a/Engine/Plugins/Experimental/DataflowMutable/Source/DataflowMutableEditor/Private/Nodes/MutableMaterialParameterNode.cpp b/Engine/Plugins/Experimental/DataflowMutable/Source/DataflowMutableEditor/Private/Nodes/MutableMaterialParameterNode.cpp
new file mode 100644
--- /dev/null
+++ b/Engine/Plugins/Experimental/DataflowMutable/Source/DataflowMutableEditor/Private/Nodes/MutableMaterialParameterNode.cpp
@@ -0,0 +1,28 @@
+﻿// Copyright Epic Games, Inc. All Rights Reserved.
+
+#include "Nodes/MutableMaterialParameterNode.h"
+#include "Materials/MaterialInterface.h"
+
+FMutableMaterialParameterNode::FMutableMaterialParameterNode(const UE::Dataflow::FNodeParameters& InParam, FGuid InGuid)
+: FDataflowNode(InParam, InGuid)
+{
+	RegisterInputConnection(&ParameterName);
+	RegisterInputConnection(&Material);
+
+	RegisterOutputConnection(&MaterialParameter);
+}
+
+
+void FMutableMaterialParameterNode::Evaluate(UE::Dataflow::FContext& Context, const FDataflowOutput* Out) const
+{
+	FMutableMaterialParameter Output;
+	Output.Name = GetValue(Context, &ParameterName);
+	Output.Material = GetValue(Context, &Material);
+	
+	if (Output.Name.IsEmpty())
+	{
+		Context.Warning(TEXT("The generated Material Parameter is empty. Make sure the parameter has a name."), this);
+	}
+	
+	SetValue(Context, Output, &MaterialParameter);
+}
diff --git a/Engine/Plugins/Experimental/DataflowMutable/Source/DataflowMutableEditor/Private/Nodes/MutableSkeletalMeshParameterNode.cpp b/Engine/Plugins/Experimental/DataflowMutable/Source/DataflowMutableEditor/Private/Nodes/MutableSkeletalMeshParameterNode.cpp
new file mode 100644
--- /dev/null
+++ b/Engine/Plugins/Experimental/DataflowMutable/Source/DataflowMutableEditor/Private/Nodes/MutableSkeletalMeshParameterNode.cpp
@@ -0,0 +1,28 @@
+﻿// Copyright Epic Games, Inc. All Rights Reserved.
+
+#include "Nodes/MutableSkeletalMeshParameterNode.h"
+#include "Engine/SkeletalMesh.h"
+
+FMutableSkeletalMeshParameterNode::FMutableSkeletalMeshParameterNode(const UE::Dataflow::FNodeParameters& InParam, FGuid InGuid)
+	: FDataflowNode(InParam, InGuid)
+{
+	RegisterInputConnection(&ParameterName);
+	RegisterInputConnection(&SkeletalMesh);
+
+	RegisterOutputConnection(&SkeletalMeshParameter);
+}
+
+
+void FMutableSkeletalMeshParameterNode::Evaluate(UE::Dataflow::FContext& Context, const FDataflowOutput* Out) const
+{
+	FMutableSkeletalMeshParameter Output;
+	Output.Name = GetValue(Context, &ParameterName);
+	Output.Mesh = GetValue(Context, &SkeletalMesh);
+
+	if (Output.Name.IsEmpty())
+	{
+		Context.Warning(TEXT("The generated Skeletal Mesh Parameter is empty. Make sure the parameter has a name."), this);
+	}
+
+	SetValue(Context, Output, &SkeletalMeshParameter);
+}
diff --git a/Engine/Plugins/Experimental/DataflowMutable/Source/DataflowMutableEditor/Private/Nodes/MutableTextureParameterNode.cpp b/Engine/Plugins/Experimental/DataflowMutable/Source/DataflowMutableEditor/Private/Nodes/MutableTextureParameterNode.cpp
new file mode 100644
--- /dev/null
+++ b/Engine/Plugins/Experimental/DataflowMutable/Source/DataflowMutableEditor/Private/Nodes/MutableTextureParameterNode.cpp
@@ -0,0 +1,28 @@
+﻿// Copyright Epic Games, Inc. All Rights Reserved.
+
+#include "Nodes/MutableTextureParameterNode.h"
+#include "Engine/Texture2D.h"
+
+FMutableTextureParameterNode::FMutableTextureParameterNode(const UE::Dataflow::FNodeParameters& InParam, FGuid InGuid)
+	: FDataflowNode(InParam, InGuid)
+{
+	RegisterInputConnection(&ParameterName);
+	RegisterInputConnection(&Texture);
+
+	RegisterOutputConnection(&TextureParameter);
+}
+
+
+void FMutableTextureParameterNode::Evaluate(UE::Dataflow::FContext& Context, const FDataflowOutput* Out) const
+{
+	FMutableTextureParameter Output;
+	Output.Name =  GetValue(Context, &ParameterName);
+	Output.Texture = GetValue(Context, &Texture);
+	
+	if (Output.Name.IsEmpty())
+	{
+		Context.Warning(TEXT("The generated Texture Parameter is empty. Make sure the parameter has a name."), this);
+	}
+
+	SetValue(Context, Output, &TextureParameter);
+}
\ No newline at end of file
diff --git a/Engine/Plugins/Experimental/DataflowMutable/Source/DataflowMutableEditor/Public/MutableDataflowParameters.h b/Engine/Plugins/Experimental/DataflowMutable/Source/DataflowMutableEditor/Public/MutableDataflowParameters.h
--- a/Engine/Plugins/Experimental/DataflowMutable/Source/DataflowMutableEditor/Public/MutableDataflowParameters.h
+++ b/Engine/Plugins/Experimental/DataflowMutable/Source/DataflowMutableEditor/Public/MutableDataflowParameters.h
@@ -2,7 +2,6 @@
 #pragma once
 
 #include "CoreMinimal.h"
-// #include "UObject/ObjectMacros.h"
 
 #include "MutableDataflowParameters.generated.h"
 
@@ -17,41 +16,41 @@ struct FMutableParameterBase
 
 	virtual ~FMutableParameterBase() = default;
 	
-	UPROPERTY(EditAnywhere, Category = "Mutable");
-	FString ParameterName;
+	UPROPERTY(EditAnywhere, Category = "Mutable")
+	FString Name;
 	
 	bool operator== (const FMutableParameterBase& Other) const
 	{
-		return this->ParameterName.ToLower() == Other.ParameterName.ToLower();
+		return this->Name.ToLower() == Other.Name.ToLower();
 	}
 };
 
 
-USTRUCT(BlueprintType, Meta = (Experimental))
+USTRUCT(BlueprintType, meta = (Experimental))
 struct FMutableSkeletalMeshParameter : public FMutableParameterBase
 {
 	GENERATED_USTRUCT_BODY()
 	
-	UPROPERTY(EditAnywhere, Category = "Mutable");
+	UPROPERTY(EditAnywhere, Category = "Mutable")
 	TObjectPtr<USkeletalMesh> Mesh;
 };
 
 
-USTRUCT(BlueprintType, Meta = (Experimental))
+USTRUCT(BlueprintType, meta = (Experimental))
 struct FMutableTextureParameter : public FMutableParameterBase
 {
 	GENERATED_USTRUCT_BODY()
 	
-	UPROPERTY(EditAnywhere, Category = "Mutable");
+	UPROPERTY(EditAnywhere, Category = "Mutable")
 	TObjectPtr<UTexture2D> Texture;
 };
 
 
-USTRUCT(BlueprintType, Meta = (Experimental))
+USTRUCT(BlueprintType, meta = (Experimental))
 struct FMutableMaterialParameter : public FMutableParameterBase
 {
 	GENERATED_USTRUCT_BODY()
 	
-	UPROPERTY(EditAnywhere, Category = "Mutable");
+	UPROPERTY(EditAnywhere, Category = "Mutable")
 	TObjectPtr<UMaterialInterface> Material;
 };
\ No newline at end of file
diff --git a/Engine/Plugins/Experimental/DataflowMutable/Source/DataflowMutableEditor/Public/Nodes/COInstanceGeneratorNode.h b/Engine/Plugins/Experimental/DataflowMutable/Source/DataflowMutableEditor/Public/Nodes/COInstanceGeneratorNode.h
--- a/Engine/Plugins/Experimental/DataflowMutable/Source/DataflowMutableEditor/Public/Nodes/COInstanceGeneratorNode.h
+++ b/Engine/Plugins/Experimental/DataflowMutable/Source/DataflowMutableEditor/Public/Nodes/COInstanceGeneratorNode.h
@@ -9,98 +9,28 @@
 #include "Engine/SkeletalMesh.h"
 #include "GeometryCollection/ManagedArrayCollection.h"
 #include "Dataflow/DataflowInputOutput.h"
-#include "Materials/MaterialInterface.h"
 #include "MutableDataflowParameters.h"
 
 #include "COInstanceGeneratorNode.generated.h"
 
-USTRUCT(Meta = (Experimental))
+USTRUCT(meta = (Experimental))
 struct FMutableGeneratedResource
 {
 	GENERATED_USTRUCT_BODY()
 
-	UPROPERTY(EditAnywhere, Category = "Mutable");
+	UPROPERTY(EditAnywhere, Category = "Mutable")
 	FString ComponentName;
 
-	UPROPERTY(EditAnywhere, Category = "Mutable");
+	UPROPERTY(EditAnywhere, Category = "Mutable")
 	TObjectPtr<USkeletalMesh> SkeletalMesh;
 };
 
 
-USTRUCT(Meta = (Experimental))
-struct FCOInstanceMeshParameterNode : public FDataflowNode
-{
-	GENERATED_USTRUCT_BODY()
-	DATAFLOW_NODE_DEFINE_INTERNAL(FCOInstanceMeshParameterNode, "Mutable Mesh Parameter", "Mutable", "")
-	
-	UPROPERTY(EditAnywhere, Category = "Mutable", meta = (DataflowInput));
-	FString ParameterName;
-
-	UPROPERTY(EditAnywhere, Category = "Mutable", meta = (DataflowInput));
-	TObjectPtr<USkeletalMesh> SkeletalMesh;
-	
-	UPROPERTY(meta = (DataflowOutput));
-	FMutableSkeletalMeshParameter MeshParameter;
-
-public:
-	
-	FCOInstanceMeshParameterNode(const UE::Dataflow::FNodeParameters& InParam, FGuid InGuid = FGuid::NewGuid());
-	
-	virtual void Evaluate(UE::Dataflow::FContext& Context, const FDataflowOutput* Out) const override;
-};
-
-
-USTRUCT(Meta = (Experimental))
-struct FCOInstanceTextureParameterNode : public FDataflowNode
-{
-	GENERATED_USTRUCT_BODY()
-	DATAFLOW_NODE_DEFINE_INTERNAL(FCOInstanceTextureParameterNode, "Mutable Texture Parameter", "Mutable", "")
-
-	UPROPERTY(EditAnywhere, Category = "Mutable", meta = (DataflowInput));
-	FString ParameterName;
-
-	UPROPERTY(EditAnywhere, Category = "Mutable", meta = (DataflowInput));
-	TObjectPtr<UTexture2D> Texture;
-	
-	UPROPERTY(meta = (DataflowOutput));
-	FMutableTextureParameter TextureParameter;
-	
-public:
-	
-	FCOInstanceTextureParameterNode(const UE::Dataflow::FNodeParameters& InParam, FGuid InGuid = FGuid::NewGuid());
-
-	virtual void Evaluate(UE::Dataflow::FContext& Context, const FDataflowOutput* Out) const override;
-};
-
-
-USTRUCT(Meta = (Experimental))
-struct FCOInstanceMaterialParameterNode : public FDataflowNode
-{
-	GENERATED_USTRUCT_BODY()
-	DATAFLOW_NODE_DEFINE_INTERNAL(FCOInstanceMaterialParameterNode, "Mutable Material Parameter", "Mutable", "")
-
-	UPROPERTY(EditAnywhere, Category = "Mutable", meta = (DataflowInput));
-	FString ParameterName;
-
-	UPROPERTY(EditAnywhere, Category = "Mutable", meta = (DataflowInput));
-	TObjectPtr<UMaterialInterface> Material;
-
-	UPROPERTY(meta = (DataflowOutput));
-	FMutableMaterialParameter MaterialParameter;
-
-public:
-	
-	FCOInstanceMaterialParameterNode(const UE::Dataflow::FNodeParameters& InParam, FGuid InGuid = FGuid::NewGuid());
-	
-	virtual void Evaluate(UE::Dataflow::FContext& Context, const FDataflowOutput* Out) const override;
-};
-
-
-USTRUCT(Meta = (Experimental))
+USTRUCT(meta = (Experimental))
 struct FCOInstanceGetComponentMesh : public FDataflowNode
 {
 	GENERATED_USTRUCT_BODY()
-	DATAFLOW_NODE_DEFINE_INTERNAL(FCOInstanceGetComponentMesh, "Get Component Mesh", "Mutable", "")
+	DATAFLOW_NODE_DEFINE_INTERNAL(FCOInstanceGetComponentMesh, "GetComponentMesh", "Mutable", "")
 
 	UPROPERTY(EditAnywhere, Category = "Mutable", meta = (DataflowInput))
 	FString ComponentName;
@@ -116,8 +46,6 @@ public:
 	FCOInstanceGetComponentMesh(const UE::Dataflow::FNodeParameters& InParam, FGuid InGuid = FGuid::NewGuid());
 	
 	virtual void Evaluate(UE::Dataflow::FContext& Context, const FDataflowOutput* Out) const override;
-	
-	virtual void OnPropertyChanged(UE::Dataflow::FContext& Context, const FPropertyChangedEvent& InPropertyChangedEvent) override;
 };
 
 
@@ -125,11 +53,11 @@ public:
  * Node designed to work as an entry point for mutable in the DataFlow graph.
  * It allows the setup of a Customizable Object Instance and the extraction of the skeletal meshes generated by its update.
  */
-USTRUCT(Meta = (Experimental))
+USTRUCT(meta = (Experimental))
 struct FCOInstanceGeneratorNode : public FDataflowNode
 {
 	GENERATED_USTRUCT_BODY()
-	DATAFLOW_NODE_DEFINE_INTERNAL(FCOInstanceGeneratorNode, "Update COI", "Mutable", "")
+	DATAFLOW_NODE_DEFINE_INTERNAL(FCOInstanceGeneratorNode, "GenerateCustomizableObjectInstance", "Mutable", "")
 
 public:
 	
@@ -143,7 +71,7 @@ private:
 	
 	// Inputs
 
-	UPROPERTY(EditAnywhere, Category = "Mutable", meta = (DataflowInput));
+	UPROPERTY(EditAnywhere, Category = "Mutable", meta = (DataflowInput))
 	TObjectPtr<UCustomizableObject> CustomizableObject;
 
 	
@@ -170,11 +98,11 @@ private:
 
 
 	// Instance generated during the runtime and used to generate the output resources
-	UPROPERTY(EditAnywhere, Category = "Mutable", meta = (DisplayName = "Generated COI"));
+	UPROPERTY(EditAnywhere, Category = "Mutable", meta = (DisplayName = "Generated COI"))
 	TObjectPtr<UCustomizableObjectInstance> CustomizableObjectInstance;
 	
 	/** Click on this button to reimport the specified USD file and regenerate the intermediary assets. */
-	UPROPERTY(EditAnywhere, Category = "Mutable", Meta = (DisplayName = "Generate Resources", ButtonImage = "Icons.Refresh"))
+	UPROPERTY(EditAnywhere, Category = "Mutable", meta = (DisplayName = "Generate Resources", ButtonImage = "Icons.Refresh"))
 	FDataflowFunctionProperty GenerateInstanceResources;
 	
 	
diff --git a/Engine/Plugins/Experimental/DataflowMutable/Source/DataflowMutableEditor/Public/Nodes/MakeMutableMaterialParametersArrayNode.h b/Engine/Plugins/Experimental/DataflowMutable/Source/DataflowMutableEditor/Public/Nodes/MakeMutableMaterialParametersArrayNode.h
new file mode 100644
--- /dev/null
+++ b/Engine/Plugins/Experimental/DataflowMutable/Source/DataflowMutableEditor/Public/Nodes/MakeMutableMaterialParametersArrayNode.h
@@ -0,0 +1,39 @@
+﻿// Copyright Epic Games, Inc. All Rights Reserved.
+#pragma once
+
+#include "Misc/Guid.h"
+#include "CoreMinimal.h"
+#include "MakeMutableParametersArrayBaseNode.h"
+#include "Dataflow/DataflowNode.h"
+#include "Dataflow/DataflowInputOutput.h"
+#include "MutableDataflowParameters.h"
+
+#include "MakeMutableMaterialParametersArrayNode.generated.h"
+
+
+USTRUCT(meta = (Experimental))
+struct FMakeMutableMaterialParametersArrayNode : public FMakeMutableParametersArrayBaseNode
+{
+	GENERATED_USTRUCT_BODY()
+	DATAFLOW_NODE_DEFINE_INTERNAL(FMakeMutableMaterialParametersArrayNode, "MakeMutableMaterialParametersArray", "Mutable", "")
+
+private:
+	UPROPERTY()
+	TArray<FMutableMaterialParameter> InputMaterialParameters;
+	
+	UPROPERTY(meta = (DataflowOutput, DisplayName = "Material Parameters"))
+	TArray<FMutableMaterialParameter> GroupedParameters;
+
+public:
+	
+	FMakeMutableMaterialParametersArrayNode(const UE::Dataflow::FNodeParameters& InParam, FGuid InGuid = FGuid::NewGuid());
+	
+	virtual void Evaluate(UE::Dataflow::FContext& Context, const FDataflowOutput* Out) const override;
+
+	// Overrides to allow the addition and removal of optional inputs
+	virtual TArray<UE::Dataflow::FPin> AddPins() override;
+	virtual bool CanRemovePin() const override;
+	virtual TArray<UE::Dataflow::FPin> GetPinsToRemove() const override;
+	virtual void OnPinRemoved(const UE::Dataflow::FPin& Pin) override;
+	virtual void PostSerialize(const FArchive& Ar) override; 
+};
\ No newline at end of file
diff --git a/Engine/Plugins/Experimental/DataflowMutable/Source/DataflowMutableEditor/Public/Nodes/MakeMutableParametersArrayBaseNode.h b/Engine/Plugins/Experimental/DataflowMutable/Source/DataflowMutableEditor/Public/Nodes/MakeMutableParametersArrayBaseNode.h
new file mode 100644
--- /dev/null
+++ b/Engine/Plugins/Experimental/DataflowMutable/Source/DataflowMutableEditor/Public/Nodes/MakeMutableParametersArrayBaseNode.h
@@ -0,0 +1,157 @@
+﻿// Copyright Epic Games, Inc. All Rights Reserved.
+#pragma once
+
+#include "Misc/Guid.h"
+#include "CoreMinimal.h"
+#include "Dataflow/DataflowNode.h"
+#include "Dataflow/DataflowInputOutput.h"
+
+#include "MakeMutableParametersArrayBaseNode.generated.h"
+
+
+USTRUCT(meta = (Experimental))
+struct FMakeMutableParametersArrayBaseNode : public FDataflowNode
+{
+	GENERATED_USTRUCT_BODY()
+	DATAFLOW_NODE_DEFINE_INTERNAL(FMakeMutableParametersArrayBaseNode, "MakeParametersArray", "Mutable", "")
+	
+public:
+	FMakeMutableParametersArrayBaseNode(const UE::Dataflow::FNodeParameters& InParam, FGuid InGuid = FGuid::NewGuid());
+
+protected:
+
+	/**
+	 * The min amount of input pins the node can have. 
+	 */
+	static constexpr uint16 MinRequiredInputsCount = 1;
+
+	template<typename T>
+	void AddDefaultInputs(TArray<T>& InParametersArray);
+
+	template<typename T>
+	void EvaluateParameterNode(const TArray<T>& InputParameters, const TArray<T>& OutputParameters, UE::Dataflow::FContext& Context, const FDataflowOutput* Out) const;
+
+	template<typename T>
+	UE::Dataflow::TConnectionReference<T> GetConnectionReference(const TArray<T>& ParametersArray, int32 Index) const;
+
+	template<typename T>
+	TArray<UE::Dataflow::FPin> AddParameterPin(TArray<T>& InOutParametersArray);
+	
+	template<typename T>
+	bool CanRemoveParameterPin(const TArray<T>& ParametersArray) const;
+
+	template<typename T>
+	TArray<UE::Dataflow::FPin> GetParameterPinsToRemove(const TArray<T>& ParametersArray) const;;
+	
+	template<typename T>
+	void OnParameterPinRemoved(TArray<T>& InParametersArray, const UE::Dataflow::FPin& Pin);
+
+	template<typename T>
+	void PostNodeSerialize(const TArray<T>& ParametersArray, const FArchive& Ar);
+
+	// FDataflowNode
+	virtual bool CanAddPin() const override;
+};
+
+
+inline FMakeMutableParametersArrayBaseNode::FMakeMutableParametersArrayBaseNode(const UE::Dataflow::FNodeParameters& InParam, FGuid InGuid)
+: Super(InParam, InGuid)
+{
+}
+
+inline bool FMakeMutableParametersArrayBaseNode::CanAddPin() const
+{
+	return true;
+}
+
+
+template <typename T>
+void FMakeMutableParametersArrayBaseNode::AddDefaultInputs(TArray<T>& InParametersArray)
+{
+	for (int32 Index = 0; Index < MinRequiredInputsCount; ++Index)
+	{
+		AddParameterPin(InParametersArray);
+	}
+}
+
+
+template <typename T>
+void FMakeMutableParametersArrayBaseNode::EvaluateParameterNode(const TArray<T>& InputParameters, const TArray<T>& OutputParameters,  UE::Dataflow::FContext& Context,
+	const FDataflowOutput* Out) const
+{
+	// Grab the values from the inputs
+	TArray<T> Output;
+	for (const T& Parameter : InputParameters)
+	{
+		Output.Add(GetValue(Context, &Parameter));
+	}
+
+	// Provide the composed array as an output
+	SetValue(Context, Output, &OutputParameters);
+}
+
+
+template<typename T>
+UE::Dataflow::TConnectionReference<T> FMakeMutableParametersArrayBaseNode::GetConnectionReference(const TArray<T>& ParametersArray, int32 Index) const
+{
+	return { &ParametersArray[Index], Index, &ParametersArray };
+}
+
+
+template <typename T>
+TArray<UE::Dataflow::FPin> FMakeMutableParametersArrayBaseNode::AddParameterPin(TArray<T>& InOutParametersArray)
+{
+	const int32 Index = InOutParametersArray.AddDefaulted();
+	UE::Dataflow::TConnectionReference<T> ConnectionRef = GetConnectionReference(InOutParametersArray, Index);
+	const FDataflowInput& Input = RegisterInputArrayConnection(ConnectionRef);
+	return { { UE::Dataflow::FPin::EDirection::INPUT, Input.GetType(), Input.GetName() } };
+}
+
+
+template <typename T>
+bool FMakeMutableParametersArrayBaseNode::CanRemoveParameterPin(const TArray<T>& ParametersArray) const
+{
+	return ParametersArray.Num() > MinRequiredInputsCount;
+}
+
+
+template <typename T>
+TArray<UE::Dataflow::FPin> FMakeMutableParametersArrayBaseNode::GetParameterPinsToRemove(const TArray<T>& ParametersArray) const
+{
+	const int32 Index = ParametersArray.Num() - 1;
+	check(ParametersArray.IsValidIndex(Index));
+	if (const FDataflowInput* const Input = FindInput(GetConnectionReference(ParametersArray, Index)))
+	{
+		return { { UE::Dataflow::FPin::EDirection::INPUT, Input->GetType(), Input->GetName() } };
+	}
+	return Super::GetPinsToRemove();
+}
+
+
+template <typename T>
+void FMakeMutableParametersArrayBaseNode::OnParameterPinRemoved(TArray<T>& InParametersArray, const UE::Dataflow::FPin& Pin)
+{
+	const int32 Index = InParametersArray.Num() - 1;
+	check(InParametersArray.IsValidIndex(Index));
+#if DO_CHECK
+		const FDataflowInput* const Input = FindInput(GetConnectionReference(InParametersArray, Index));
+		check(Input);
+		check(Input->GetName() == Pin.Name);
+		check(Input->GetType() == Pin.Type);
+#endif
+	InParametersArray.SetNum(Index);
+}
+
+
+template <typename T>
+void FMakeMutableParametersArrayBaseNode::PostNodeSerialize(const TArray<T>& ParametersArray, const FArchive& Ar)
+{
+	if (Ar.IsLoading())
+	{
+		// Make sure the node does have the same amount of inputs as when saved
+		for (int32 Index = 0; Index < ParametersArray.Num(); ++Index)
+		{
+			FindOrRegisterInputArrayConnection(GetConnectionReference(ParametersArray, Index));
+		}
+	}
+}
diff --git a/Engine/Plugins/Experimental/DataflowMutable/Source/DataflowMutableEditor/Public/Nodes/MakeMutableSkeletalMeshParametersArrayNode.h b/Engine/Plugins/Experimental/DataflowMutable/Source/DataflowMutableEditor/Public/Nodes/MakeMutableSkeletalMeshParametersArrayNode.h
new file mode 100644
--- /dev/null
+++ b/Engine/Plugins/Experimental/DataflowMutable/Source/DataflowMutableEditor/Public/Nodes/MakeMutableSkeletalMeshParametersArrayNode.h
@@ -0,0 +1,39 @@
+// Copyright Epic Games, Inc. All Rights Reserved.
+#pragma once
+
+#include "Misc/Guid.h"
+#include "CoreMinimal.h"
+#include "MakeMutableParametersArrayBaseNode.h"
+#include "Dataflow/DataflowNode.h"
+#include "Dataflow/DataflowInputOutput.h"
+#include "MutableDataflowParameters.h"
+
+#include "MakeMutableSkeletalMeshParametersArrayNode.generated.h"
+
+
+USTRUCT(meta = (Experimental))
+struct FMakeMutableSkeletalMeshParametersArrayNode : public FMakeMutableParametersArrayBaseNode
+{
+	GENERATED_USTRUCT_BODY()
+	DATAFLOW_NODE_DEFINE_INTERNAL(FMakeMutableSkeletalMeshParametersArrayNode, "MakeMutableSkeletalMeshParametersArray", "Mutable", "")
+
+private:
+	UPROPERTY()
+	TArray<FMutableSkeletalMeshParameter> InputSkeletalMeshParameters;
+	
+	UPROPERTY(meta = (DataflowOutput, DisplayName = "Skeletal Mesh Parameters"))
+	TArray<FMutableSkeletalMeshParameter> GroupedParameters;
+
+public:
+	
+	FMakeMutableSkeletalMeshParametersArrayNode(const UE::Dataflow::FNodeParameters& InParam, FGuid InGuid = FGuid::NewGuid());
+	
+	virtual void Evaluate(UE::Dataflow::FContext& Context, const FDataflowOutput* Out) const override;
+	
+	// Overrides to allow the addition and removal of optional inputs
+	virtual TArray<UE::Dataflow::FPin> AddPins() override;
+	virtual bool CanRemovePin() const override;
+	virtual TArray<UE::Dataflow::FPin> GetPinsToRemove() const override;
+	virtual void OnPinRemoved(const UE::Dataflow::FPin& Pin) override;
+	virtual void PostSerialize(const FArchive& Ar) override; 
+};
\ No newline at end of file
diff --git a/Engine/Plugins/Experimental/DataflowMutable/Source/DataflowMutableEditor/Public/Nodes/MakeMutableTextureParametersArrayNode.h b/Engine/Plugins/Experimental/DataflowMutable/Source/DataflowMutableEditor/Public/Nodes/MakeMutableTextureParametersArrayNode.h
new file mode 100644
--- /dev/null
+++ b/Engine/Plugins/Experimental/DataflowMutable/Source/DataflowMutableEditor/Public/Nodes/MakeMutableTextureParametersArrayNode.h
@@ -0,0 +1,39 @@
+﻿// Copyright Epic Games, Inc. All Rights Reserved.
+#pragma once
+
+#include "Misc/Guid.h"
+#include "CoreMinimal.h"
+#include "MakeMutableParametersArrayBaseNode.h"
+#include "Dataflow/DataflowNode.h"
+#include "Dataflow/DataflowInputOutput.h"
+#include "MutableDataflowParameters.h"
+
+#include "MakeMutableTextureParametersArrayNode.generated.h"
+
+
+USTRUCT(meta = (Experimental))
+struct FMakeMutableTextureParametersArrayNode : public FMakeMutableParametersArrayBaseNode
+{
+	GENERATED_USTRUCT_BODY()
+	DATAFLOW_NODE_DEFINE_INTERNAL(FMakeMutableTextureParametersArrayNode, "MakeMutableTextureParametersArray", "Mutable", "")
+
+private:
+	UPROPERTY()
+	TArray<FMutableTextureParameter> InputTextureParameters;
+	
+	UPROPERTY(meta = (DataflowOutput,  DisplayName = "Texture Parameters"))
+	TArray<FMutableTextureParameter> GroupedParameters;
+
+public:
+	
+	FMakeMutableTextureParametersArrayNode(const UE::Dataflow::FNodeParameters& InParam, FGuid InGuid = FGuid::NewGuid());
+	
+	virtual void Evaluate(UE::Dataflow::FContext& Context, const FDataflowOutput* Out) const override;
+
+	// Overrides to allow the addition and removal of optional inputs
+	virtual TArray<UE::Dataflow::FPin> AddPins() override;
+	virtual bool CanRemovePin() const override;
+	virtual TArray<UE::Dataflow::FPin> GetPinsToRemove() const override;
+	virtual void OnPinRemoved(const UE::Dataflow::FPin& Pin) override;
+	virtual void PostSerialize(const FArchive& Ar) override; 
+};
\ No newline at end of file
diff --git a/Engine/Plugins/Experimental/DataflowMutable/Source/DataflowMutableEditor/Public/Nodes/MutableMaterialParameterNode.h b/Engine/Plugins/Experimental/DataflowMutable/Source/DataflowMutableEditor/Public/Nodes/MutableMaterialParameterNode.h
new file mode 100644
--- /dev/null
+++ b/Engine/Plugins/Experimental/DataflowMutable/Source/DataflowMutableEditor/Public/Nodes/MutableMaterialParameterNode.h
@@ -0,0 +1,31 @@
+﻿// Copyright Epic Games, Inc. All Rights Reserved.
+#pragma once
+
+#include "Dataflow/DataflowNode.h"
+#include "Misc/Guid.h"
+#include "CoreMinimal.h"
+#include "MutableDataflowParameters.h"
+
+#include "MutableMaterialParameterNode.generated.h"
+
+USTRUCT(meta = (Experimental))
+struct FMutableMaterialParameterNode : public FDataflowNode
+{
+	GENERATED_USTRUCT_BODY()
+	DATAFLOW_NODE_DEFINE_INTERNAL(FMutableMaterialParameterNode, "MutableMaterialParameter", "Mutable", "")
+
+	UPROPERTY(EditAnywhere, Category = "Mutable", meta = (DataflowInput))
+	FString ParameterName;
+
+	UPROPERTY(EditAnywhere, Category = "Mutable", meta = (DataflowInput))
+	TObjectPtr<UMaterialInterface> Material;
+
+	UPROPERTY(meta = (DataflowOutput))
+	FMutableMaterialParameter MaterialParameter;
+
+public:
+	
+	FMutableMaterialParameterNode(const UE::Dataflow::FNodeParameters& InParam, FGuid InGuid = FGuid::NewGuid());
+	
+	virtual void Evaluate(UE::Dataflow::FContext& Context, const FDataflowOutput* Out) const override;
+};
\ No newline at end of file
diff --git a/Engine/Plugins/Experimental/DataflowMutable/Source/DataflowMutableEditor/Public/Nodes/MutableSkeletalMeshParameterNode.h b/Engine/Plugins/Experimental/DataflowMutable/Source/DataflowMutableEditor/Public/Nodes/MutableSkeletalMeshParameterNode.h
new file mode 100644
--- /dev/null
+++ b/Engine/Plugins/Experimental/DataflowMutable/Source/DataflowMutableEditor/Public/Nodes/MutableSkeletalMeshParameterNode.h
@@ -0,0 +1,32 @@
+﻿// Copyright Epic Games, Inc. All Rights Reserved.
+#pragma once
+
+#include "Dataflow/DataflowNode.h"
+#include "MutableDataflowParameters.h"
+#include "Misc/Guid.h"
+#include "CoreMinimal.h"
+
+#include "MutableSkeletalMeshParameterNode.generated.h"
+
+
+USTRUCT(meta = (Experimental))
+struct FMutableSkeletalMeshParameterNode : public FDataflowNode
+{
+	GENERATED_USTRUCT_BODY()
+	DATAFLOW_NODE_DEFINE_INTERNAL(FMutableSkeletalMeshParameterNode, "MutableSkeletalMeshParameter", "Mutable", "")
+	
+	UPROPERTY(EditAnywhere, Category = "Mutable", meta = (DataflowInput))
+	FString ParameterName;
+
+	UPROPERTY(EditAnywhere, Category = "Mutable", meta = (DataflowInput))
+	TObjectPtr<USkeletalMesh> SkeletalMesh;
+	
+	UPROPERTY(meta = (DataflowOutput))
+	FMutableSkeletalMeshParameter SkeletalMeshParameter;
+
+public:
+	
+	FMutableSkeletalMeshParameterNode(const UE::Dataflow::FNodeParameters& InParam, FGuid InGuid = FGuid::NewGuid());
+	
+	virtual void Evaluate(UE::Dataflow::FContext& Context, const FDataflowOutput* Out) const override;
+};
\ No newline at end of file
diff --git a/Engine/Plugins/Experimental/DataflowMutable/Source/DataflowMutableEditor/Public/Nodes/MutableTextureParameterNode.h b/Engine/Plugins/Experimental/DataflowMutable/Source/DataflowMutableEditor/Public/Nodes/MutableTextureParameterNode.h
new file mode 100644
--- /dev/null
+++ b/Engine/Plugins/Experimental/DataflowMutable/Source/DataflowMutableEditor/Public/Nodes/MutableTextureParameterNode.h
@@ -0,0 +1,31 @@
+﻿// Copyright Epic Games, Inc. All Rights Reserved.
+#pragma once
+
+#include "Dataflow/DataflowNode.h"
+#include "Misc/Guid.h"
+#include "CoreMinimal.h"
+#include "MutableDataflowParameters.h"
+
+#include "MutableTextureParameterNode.generated.h"
+
+USTRUCT(meta = (Experimental))
+struct FMutableTextureParameterNode : public FDataflowNode
+{
+	GENERATED_USTRUCT_BODY()
+	DATAFLOW_NODE_DEFINE_INTERNAL(FMutableTextureParameterNode, "MutableTextureParameter", "Mutable", "")
+
+	UPROPERTY(EditAnywhere, Category = "Mutable", meta = (DataflowInput))
+	FString ParameterName;
+
+	UPROPERTY(EditAnywhere, Category = "Mutable", meta = (DataflowInput))
+	TObjectPtr<UTexture2D> Texture;
+	
+	UPROPERTY(meta = (DataflowOutput))
+	FMutableTextureParameter TextureParameter;
+	
+public:
+	
+	FMutableTextureParameterNode(const UE::Dataflow::FNodeParameters& InParam, FGuid InGuid = FGuid::NewGuid());
+
+	virtual void Evaluate(UE::Dataflow::FContext& Context, const FDataflowOutput* Out) const override;
+};
\ No newline at end of file
