[mutable] Fixed incorrect comment.

#rnx

diff --git a/Engine/Plugins/Experimental/DataflowMutable/Source/DataflowMutableEditor/Public/Nodes/COInstanceGeneratorNode.h b/Engine/Plugins/Experimental/DataflowMutable/Source/DataflowMutableEditor/Public/Nodes/COInstanceGeneratorNode.h
--- a/Engine/Plugins/Experimental/DataflowMutable/Source/DataflowMutableEditor/Public/Nodes/COInstanceGeneratorNode.h
+++ b/Engine/Plugins/Experimental/DataflowMutable/Source/DataflowMutableEditor/Public/Nodes/COInstanceGeneratorNode.h
@@ -100,8 +100,8 @@ private:
 	// Instance generated during the runtime and used to generate the output resources
 	UPROPERTY(EditAnywhere, Category = "Mutable", meta = (DisplayName = "Generated COI"))
 	TObjectPtr<UCustomizableObjectInstance> CustomizableObjectInstance;
-	
-	/** Click on this button to reimport the specified USD file and regenerate the intermediary assets. */
+
+	/** Press this button to trigger first a CO compilation and then a COI update. This will generate the resources needed by the outputs of the node. */
 	UPROPERTY(EditAnywhere, Category = "Mutable", meta = (DisplayName = "Generate Resources", ButtonImage = "Icons.Refresh"))
 	FDataflowFunctionProperty GenerateInstanceResources;
 	
